HttpInsights - a high level HTTP trace insights module build for but not restricted to IAX streaming

#rnx

#Synced-CL 50152864

diff --git a/Engine/Plugins/Experimental/HttpInsights/HttpInsights.uplugin b/Engine/Plugins/Experimental/HttpInsights/HttpInsights.uplugin
new file mode 100644
--- /dev/null
+++ b/Engine/Plugins/Experimental/HttpInsights/HttpInsights.uplugin
@@ -0,0 +1,28 @@
+{
+	"FileVersion": 3,
+	"Version": 1,
+	"VersionName": "1.0",
+	"FriendlyName": "HTTP Insights",
+	"Description": "Allows capturing HTTP traffic",
+	"Category": "Insights",
+	"CreatedBy": "Epic Games, Inc.",
+	"CreatedByURL": "https://epicgames.com",
+	"DocsURL": "",
+	"MarketplaceURL": "",
+	"SupportURL": "",
+	"EnabledByDefault": true,
+	"CanContainContent": false,
+	"IsBetaVersion": false,
+	"Installed": false,
+	"SupportedPrograms": [ "UnrealInsights" ],
+	"Modules": [
+		{
+			"Name": "HttpInsights",
+			"Type": "EditorAndProgram",
+			"LoadingPhase": "Default",
+			"ProgramAllowList": [
+				"UnrealInsights"
+			]
+		}
+	]
+}
diff --git a/Engine/Plugins/Experimental/HttpInsights/Source/HttpInsights.Build.cs b/Engine/Plugins/Experimental/HttpInsights/Source/HttpInsights.Build.cs
new file mode 100644
--- /dev/null
+++ b/Engine/Plugins/Experimental/HttpInsights/Source/HttpInsights.Build.cs
@@ -0,0 +1,26 @@
+// Copyright Epic Games, Inc. All Rights Reserved.
+
+namespace UnrealBuildTool.Rules
+{
+	public class HttpInsights : ModuleRules
+	{
+		public HttpInsights(ReadOnlyTargetRules Target) : base(Target)
+		{
+			PrivateDependencyModuleNames.AddRange(
+			new string[]
+			{
+				"Core",
+				"CoreUObject",
+				"SlateCore",
+				"Slate",
+				"InputCore",
+				"ToolMenus",
+				"ToolWidgets",
+				"TraceAnalysis",
+				"TraceServices",
+				"TraceInsights",
+				"TraceInsightsCore",
+			});
+		}
+	}
+}
diff --git a/Engine/Plugins/Experimental/HttpInsights/Source/Private/Analyzer.cpp b/Engine/Plugins/Experimental/HttpInsights/Source/Private/Analyzer.cpp
new file mode 100644
--- /dev/null
+++ b/Engine/Plugins/Experimental/HttpInsights/Source/Private/Analyzer.cpp
@@ -0,0 +1,172 @@
+// Copyright Epic Games, Inc. All Rights Reserved.
+
+#include "Model.h"
+
+#include "Trace/Analyzer.h"
+#include "TraceServices/Model/Definitions.h"
+#include "TraceServices/Model/Strings.h"
+#include "TraceServices/Model/MetadataProvider.h"
+
+namespace UE::HttpInsights
+{
+
+using namespace TraceServices;
+
+////////////////////////////////////////////////////////////////////////////////////////////////////
+enum class EHttpTraceEventId : uint16
+{
+	DispatcherCreated,
+	CategoryCreated,
+	RequestStarted,
+	ChunkRangeAdded,
+	RequestCompleted,
+	PackageIdMapped,
+};
+
+////////////////////////////////////////////////////////////////////////////////////////////////////
+class FHttpAnalyzer
+	: public UE::Trace::IAnalyzer 
+{
+public:
+	FHttpAnalyzer(IAnalysisSession& InSession, IHttpLogModel& InLogModel)
+		: Session(InSession)
+		, LogModel(InLogModel)
+	{ }
+
+	virtual void OnAnalysisBegin(const FOnAnalysisContext& Context) override;
+	virtual bool OnEvent(uint16 RouteId, EStyle Style, const FOnEventContext& Context) override;
+
+private:
+	IAnalysisSession&	Session;
+	IHttpLogModel&		LogModel;
+};
+
+////////////////////////////////////////////////////////////////////////////////////////////////////
+void FHttpAnalyzer::OnAnalysisBegin(const FOnAnalysisContext& Context)
+{
+	Trace::IAnalyzer::FInterfaceBuilder& Builder = Context.InterfaceBuilder;
+	Builder.RouteEvent(uint16(EHttpTraceEventId::DispatcherCreated),	"Http",		"DispatcherCreated");
+	Builder.RouteEvent(uint16(EHttpTraceEventId::CategoryCreated),		"Http", 	"CategoryCreated");
+	Builder.RouteEvent(uint16(EHttpTraceEventId::RequestStarted),		"Http", 	"RequestStarted");
+	Builder.RouteEvent(uint16(EHttpTraceEventId::ChunkRangeAdded),		"Http", 	"ChunkRangeAdded");
+	Builder.RouteEvent(uint16(EHttpTraceEventId::RequestCompleted),		"Http", 	"RequestCompleted");
+	Builder.RouteEvent(uint16(EHttpTraceEventId::PackageIdMapped),		"Package",	"PackageMapping");
+}
+
+bool FHttpAnalyzer::OnEvent(uint16 EventIdValue, EStyle Style, const FOnEventContext& Context)
+{
+	const auto& EventData			= Context.EventData;
+	const EHttpTraceEventId EventId = static_cast<EHttpTraceEventId>(EventIdValue);
+
+	switch(EventId)
+	{
+	case EHttpTraceEventId::DispatcherCreated:
+	{
+		FAnsiStringView Name;
+		EventData.GetString("Name", Name);
+		LogModel.ProcesEvent(FHttpDispatcherCreated
+		{
+			.Dispatcher = EventData.GetValue<uint64>("Dispatcher"),
+			.Name		= FAnsiString(Name)
+		});
+	}
+	break;
+
+	case EHttpTraceEventId::CategoryCreated:
+	{
+		FAnsiStringView Name;
+		EventData.GetString("Name", Name);
+		LogModel.ProcesEvent(FHttpCategoryCreated
+		{
+			.Category	= EventData.GetValue<uint8>("Category"),
+			.Name		= FAnsiString(Name)
+		});
+	}
+	break;
+
+	case EHttpTraceEventId::RequestStarted:
+	{
+		const uint64 Cycles = EventData.GetValue<uint64>("StartTime");
+		FAnsiStringView Url;
+		EventData.GetString("Url", Url);
+
+		LogModel.ProcesEvent(FHttpRequestStarted
+		{
+			.Dispatcher = EventData.GetValue<uint64>("Dispatcher"),
+			.Request	= EventData.GetValue<uint64>("Request"),
+			.StartTime	= Context.EventTime.AsSeconds(Cycles),
+			.Priority	= EventData.GetValue<int32>("Priority"),
+			.Category	= EventData.GetValue<uint32>("Category"),
+			.Url		= FAnsiString(Url)
+		});
+	}
+	break;
+
+	case EHttpTraceEventId::ChunkRangeAdded:
+	{
+		FHttpChunkRangeAdded ChunkRangeAdded = FHttpChunkRangeAdded
+		{
+			.Request	= EventData.GetValue<uint64>("Request"),
+			.Start		= EventData.GetValue<uint32>("Start"),
+			.End		= EventData.GetValue<uint32>("End")
+		};
+
+		const uint32 ChunkIdSize = uint32(sizeof(FIoChunkId));
+		const TArrayReader<uint8>& Reader = EventData.GetArray<uint8>("ChunkId");
+		if (Reader.Num() == ChunkIdSize) 
+		{
+			FMemory::Memcpy(&ChunkRangeAdded.ChunkId, Reader.GetData(), ChunkIdSize);
+		}
+		LogModel.ProcesEvent(MoveTemp(ChunkRangeAdded));
+	}
+	break;
+
+	case EHttpTraceEventId::RequestCompleted:
+	{
+		const uint64 Cycles = EventData.GetValue<uint64>("CompletionTime");
+		FAnsiStringView Host;
+		EventData.GetString("Host", Host);
+
+		LogModel.ProcesEvent(FHttpRequestCompleted
+		{
+			.Request		= EventData.GetValue<uint64>("Request"),
+			.CompletionTime	= Context.EventTime.AsSeconds(Cycles),
+			.StatusCode		= EventData.GetValue<uint32>("StatusCode"),
+			.ContentLength	= EventData.GetValue<uint32>("ContentLength"),
+			.Host			= FAnsiString(Host)
+		});
+	}
+	break;
+
+	case EHttpTraceEventId::PackageIdMapped:
+	{
+		if (const IDefinitionProvider* DefinitionProvider = TraceServices::ReadDefinitionProvider(Session))
+		{
+			FPackageId PackageId		= FPackageId::FromValue(EventData.GetValue<uint64>("Id"));
+			const auto PackageNameRef	= EventData.GetReferenceValue<uint32>("Package");
+			const auto PackageName		= DefinitionProvider->Get<FStringDefinition>(PackageNameRef);
+
+			if (PackageName && *PackageName->Display != '\0')
+			{
+				LogModel.ProcesEvent(FPackageIdMapped
+				{
+					.PackageId		= PackageId,
+					.PackageName	= PackageName->Display
+				});
+			}
+		}
+	}
+	break;
+
+	}
+
+	return true;
+}
+
+////////////////////////////////////////////////////////////////////////////////////////////////////
+TUniquePtr<UE::Trace::IAnalyzer> MakeHttpAnalyzer(IAnalysisSession& Session, IHttpLogModel& LogModel)
+{
+	return MakeUnique<FHttpAnalyzer>(Session, LogModel);
+}
+
+} // namespace UE:Http::Insights
diff --git a/Engine/Plugins/Experimental/HttpInsights/Source/Private/Model.cpp b/Engine/Plugins/Experimental/HttpInsights/Source/Private/Model.cpp
new file mode 100644
--- /dev/null
+++ b/Engine/Plugins/Experimental/HttpInsights/Source/Private/Model.cpp
@@ -0,0 +1,251 @@
+// Copyright Epic Games, Inc. All Rights Reserved.
+
+#include "Model.h"
+
+#include "Async/Mutex.h"
+#include "Async/SharedLock.h"
+#include "Async/SharedMutex.h"
+#include "Async/UniqueLock.h"
+#include "Containers/PagedArray.h"
+#include "HttpInsights/ModuleInterface.h"
+
+namespace UE::HttpInsights
+{
+
+using namespace TraceServices;
+
+////////////////////////////////////////////////////////////////////////////////
+static FPackageId GetPackageId(const FIoChunkId& ChunkId)
+{
+	FPackageId Out;
+	FMemory::Memcpy(&Out, &ChunkId, sizeof(FPackageId));
+	return Out;
+}
+
+////////////////////////////////////////////////////////////////////////////////////////////////////
+class FPackageNameCache
+{
+	using FLookup = TMap<FPackageId, const TCHAR*>;
+
+public:
+	void Put(const FPackageId& PackageId, const TCHAR* PackageName)
+	{
+		TUniqueLock Lock(Mutex);
+		Lookup.Add(PackageId, PackageName);
+	}
+
+	const TCHAR* Get(const FIoChunkId& ChunkId) const
+	{
+		const FPackageId PackageId = GetPackageId(ChunkId);
+		TUniqueLock Lock(Mutex);
+		return Lookup.FindRef(PackageId);
+	}
+
+private:
+	FLookup			Lookup;
+	mutable FMutex	Mutex;
+};
+
+////////////////////////////////////////////////////////////////////////////////////////////////////
+class FHttpLogModel
+	: public IHttpLogModel
+{
+	struct FCategory
+	{
+		FString Name;
+		uint8	Value;
+	};
+
+	using FPendingReqeusts	= TMap<uint64, FHttpRequestStarted>;
+	using FCategories		= TPagedArray<FCategory>;
+	using FPackageMappings	= TMap<FPackageId, const TCHAR*>;
+
+public:
+							FHttpLogModel(IAnalysisSession& InSession);
+	// Writer
+	virtual void			ProcesEvent(FHttpDispatcherCreated&& Event) override;
+	virtual void			ProcesEvent(FHttpCategoryCreated&&) override;
+	virtual void			ProcesEvent(FHttpRequestStarted&& Event) override;
+	virtual void			ProcesEvent(FHttpChunkRangeAdded&& Event) override;
+	virtual void			ProcesEvent(FHttpRequestCompleted&& Event) override;
+	virtual void			ProcesEvent(FPackageIdMapped&&) override;
+	// Reader
+	virtual int32			IterateLog(int32 StartIndex, IHttpLogModel::FReadLogFunc&&) const override;
+	virtual const TCHAR*	GetPackageName(const FIoChunkId& ChunkId) const override;
+
+private:
+	const FHttpDispatcher&	FindDispatcher(uint64 Handle);
+	const FCategory&		FindCategory(uint8 CategoryValue) const;
+
+	IAnalysisSession&				Session;
+	TPagedArray<FHttpDispatcher>	Dispatchers;
+	TPagedArray<FHttpRequest>		Requests;
+	TPagedArray<FHttpChunkRange>	ChunkRanges;
+	FCategories						Categories;
+	FPackageNameCache				PackageNameCache;
+	FPendingReqeusts				PendingRequests;
+	mutable FSharedMutex			Mutex;
+	uint32							SeqNo = 0;
+};
+
+////////////////////////////////////////////////////////////////////////////////////////////////////
+FHttpLogModel::FHttpLogModel(IAnalysisSession& InSession)
+	: Session(InSession)
+{
+	Dispatchers.Add(FHttpDispatcher
+	{
+		.Handle = MAX_uint64,
+		.Name	= TEXT("Unknown")
+	});
+	
+	Categories.Add(FCategory
+	{
+		.Name		= TEXT("Unknown"),
+		.Value		= MAX_uint8
+	});
+}
+
+void FHttpLogModel::ProcesEvent(FHttpDispatcherCreated&& Event)
+{
+	Dispatchers.Add(FHttpDispatcher
+	{
+		.Handle = Event.Dispatcher,
+		.Name	= FString(Event.Name)
+	});
+}
+
+void FHttpLogModel::ProcesEvent(FHttpCategoryCreated&& Event)
+{
+	Categories.Add(FCategory
+	{
+		.Name		= FString(Event.Name),
+		.Value		= Event.Category
+	});
+}
+
+void FHttpLogModel::ProcesEvent(FHttpRequestStarted&& Event)
+{
+	if (FHttpRequestStarted* Pending = PendingRequests.Find(Event.Request); Pending == nullptr)
+	{
+		PendingRequests.Add(Event.Request, MoveTemp(Event));
+	}
+	else
+	{
+		UE_LOG(LogHttpInsights, Warning, TEXT("Request already pending"));
+	}
+}
+
+void FHttpLogModel::ProcesEvent(FHttpChunkRangeAdded&& Event)
+{
+	if (FHttpRequestStarted* Pending = PendingRequests.Find(Event.Request); Pending != nullptr)
+	{
+		FHttpChunkRange* Next = Pending->ChunkRanges;
+		Pending->ChunkRanges = &ChunkRanges.Add_GetRef(FHttpChunkRange
+		{
+			.Next			= Next,
+			.ChunkId		= Event.ChunkId,
+			.Start			= Event.Start,
+			.End			= Event.End
+		});
+	}
+	else
+	{
+		UE_LOG(LogHttpInsights, Warning, TEXT("Failed to append chunk range info to pending request"));
+	}
+}
+
+void FHttpLogModel::ProcesEvent(FHttpRequestCompleted&& Event)
+{
+	if (FHttpRequestStarted* Pending = PendingRequests.Find(Event.Request); Pending != nullptr)
+	{
+		const uint8 CategoryValue = static_cast<uint8>(Pending->Category);
+		const FCategory& Category = FindCategory(CategoryValue);
+		{
+			TUniqueLock Lock(Mutex);
+			Requests.Add(FHttpRequest
+			{
+				.Dispatcher		= FindDispatcher(Pending->Dispatcher),
+				.StartTime		= Pending->StartTime,
+				.CompletionTime	= Event.CompletionTime,
+				.StatusCode		= Event.StatusCode,
+				.ContentLength	= Event.ContentLength,
+				.SeqNo			= ++SeqNo,
+				.ChunkRanges	= MoveTemp(Pending->ChunkRanges),
+				.Host			= FString(Event.Host),
+				.Url			= FString(Pending->Url),
+				.CategoryName	= &Category.Name,
+				.Category		= CategoryValue
+			});
+		}
+		PendingRequests.Remove(Event.Request);
+	}
+	else
+	{
+		UE_LOG(LogHttpInsights, Warning, TEXT("Unknown completed request"));
+	}
+}
+
+void FHttpLogModel::ProcesEvent(FPackageIdMapped&& Event)
+{
+	PackageNameCache.Put(Event.PackageId, Event.PackageName);
+}
+
+int32 FHttpLogModel::IterateLog(int32 StartIndex, IHttpLogModel::FReadLogFunc&& Read) const
+{
+	int32 Index = StartIndex < 0 ? 0 : StartIndex;
+
+	TSharedLock Lock(Mutex);
+
+	if (Requests.IsEmpty())
+	{
+		return INDEX_NONE;
+	}
+
+	for (int32 Num = Requests.Num(); Index < Num; ++Index)
+	{
+		Read(Requests[Index]);
+	}
+
+	return Index;
+}
+
+const FHttpDispatcher& FHttpLogModel::FindDispatcher(uint64 Handle)
+{
+	for (const FHttpDispatcher& Dispatcher : Dispatchers)
+	{
+		if (Dispatcher.Handle == Handle)
+		{
+			return Dispatcher;
+		}
+	}
+
+	return Dispatchers[0]; // Unknown
+}
+
+const FHttpLogModel::FCategory& FHttpLogModel::FindCategory(uint8 CategoryValue) const
+{
+	for (const FCategory& Category : Categories)
+	{
+		if (Category.Value == CategoryValue)
+		{
+			return Category;
+		}
+	}
+
+	return Categories[0]; // Unknown
+}
+
+const TCHAR* FHttpLogModel::GetPackageName(const FIoChunkId& ChunkId) const
+{
+	return PackageNameCache.Get(ChunkId);
+}
+
+////////////////////////////////////////////////////////////////////////////////////////////////////
+TSharedPtr<IHttpLogModel> MakeHttpLogModel(IAnalysisSession& Session)
+{
+	return MakeShared<FHttpLogModel>(Session);
+}
+
+} // namespace UE:HttpInsights
+
+#undef LOCTEXT_NAMESPACE
diff --git a/Engine/Plugins/Experimental/HttpInsights/Source/Private/Model.h b/Engine/Plugins/Experimental/HttpInsights/Source/Private/Model.h
new file mode 100644
--- /dev/null
+++ b/Engine/Plugins/Experimental/HttpInsights/Source/Private/Model.h
@@ -0,0 +1,121 @@
+// Copyright Epic Games, Inc. All Rights Reserved.
+
+#pragma once
+
+#include "Containers/AnsiString.h"
+#include "IO/IoChunkId.h"
+#include "IO/PackageId.h"
+#include "Math/NumericLimits.h"
+#include "Templates/Function.h"
+#include "TraceServices/Model/AnalysisSession.h"
+
+namespace UE::HttpInsights
+{
+
+////////////////////////////////////////////////////////////////////////////////////////////////////
+struct FHttpDispatcher
+{
+	const uint64	Handle = 0;
+	const FString	Name;
+};
+
+////////////////////////////////////////////////////////////////////////////////////////////////////
+struct FHttpChunkRange
+{
+	FHttpChunkRange*	Next		= nullptr;
+	FIoChunkId			ChunkId		= FIoChunkId::InvalidChunkId;
+	uint32				Start		= MAX_uint32;
+	uint32				End			= MIN_uint32;
+};
+
+////////////////////////////////////////////////////////////////////////////////////////////////////
+struct FHttpRequest
+{
+	const FHttpDispatcher&	Dispatcher;
+	double					StartTime		= 0.0;
+	double					CompletionTime	= 0.0;
+	uint32					StatusCode		= 0;
+	uint32					ContentLength	= 0;
+	uint32					SeqNo			= 0;
+	FHttpChunkRange*		ChunkRanges		= nullptr;
+	FString					Host;
+	FString					Url;
+	const FString*			CategoryName	= nullptr;
+	uint8					Category		= MAX_uint8;
+};
+
+////////////////////////////////////////////////////////////////////////////////////////////////////
+struct FHttpDispatcherCreated
+{
+	uint64		Dispatcher = 0;
+	FAnsiString	Name;
+};
+
+////////////////////////////////////////////////////////////////////////////////////////////////////
+struct FHttpCategoryCreated
+{
+	uint8		Category = MAX_uint8;
+	FAnsiString	Name;
+};
+
+////////////////////////////////////////////////////////////////////////////////////////////////////
+struct FHttpRequestStarted
+{
+	uint64					Dispatcher		= 0;
+	uint64					Request			= 0;
+	double					StartTime		= 0.0;
+	int32					Priority		= 0;
+	uint32					Category		= 0;
+	FHttpChunkRange*		ChunkRanges		= nullptr;
+	FAnsiString				Url;
+};
+
+////////////////////////////////////////////////////////////////////////////////////////////////////
+struct FHttpChunkRangeAdded
+{
+	uint64		Request		= 0;
+	uint32		Start		= MAX_uint32;
+	uint32		End			= MIN_uint32;
+	FIoChunkId	ChunkId		= FIoChunkId::InvalidChunkId;
+};
+
+////////////////////////////////////////////////////////////////////////////////////////////////////
+struct FHttpRequestCompleted
+{
+	uint64		Request			= 0;
+	double		CompletionTime	= 0.0;
+	uint32		StatusCode		= 0;
+	uint32		ContentLength	= 0;
+	FAnsiString	Host;
+};
+
+////////////////////////////////////////////////////////////////////////////////////////////////////
+struct FPackageIdMapped
+{
+	FPackageId		PackageId;
+	const TCHAR*	PackageName = nullptr;
+};
+
+////////////////////////////////////////////////////////////////////////////////////////////////////
+class IHttpLogModel
+	: public TraceServices::IProvider
+{
+public:
+	inline static const FName ProviderName = TEXT("HttpLogProvider");
+	using FReadLogFunc = TFunctionRef<void(const FHttpRequest&)>;
+
+	virtual					~IHttpLogModel() = default;
+	virtual void			ProcesEvent(FHttpDispatcherCreated&&) = 0;
+	virtual void			ProcesEvent(FHttpCategoryCreated&&) = 0;
+	virtual void 			ProcesEvent(FHttpRequestStarted&&) = 0;
+	virtual void 			ProcesEvent(FHttpChunkRangeAdded&&) = 0;
+	virtual void 			ProcesEvent(FHttpRequestCompleted&&) = 0;
+	virtual void 			ProcesEvent(FPackageIdMapped&&) = 0;
+	virtual int32			IterateLog(int32 StartIndex, FReadLogFunc&&) const = 0;
+	virtual const TCHAR*	GetPackageName(const FIoChunkId&) const = 0;
+};
+
+////////////////////////////////////////////////////////////////////////////////////////////////////
+TSharedPtr<IHttpLogModel> MakeHttpLogModel(TraceServices::IAnalysisSession& Session);
+
+} // namespace UE:HttpInsights
diff --git a/Engine/Plugins/Experimental/HttpInsights/Source/Private/Module.cpp b/Engine/Plugins/Experimental/HttpInsights/Source/Private/Module.cpp
new file mode 100644
--- /dev/null
+++ b/Engine/Plugins/Experimental/HttpInsights/Source/Private/Module.cpp
@@ -0,0 +1,137 @@
+// Copyright Epic Games, Inc. All Rights Reserved.
+
+#include "HttpInsights/ModuleInterface.h"
+#include "Model.h"
+
+#include "Features/IModularFeatures.h"
+#include "Insights/IUnrealInsightsModule.h"
+#include "Insights/ITimingViewExtender.h"
+#include "Modules/ModuleManager.h"
+#include "Trace/Analyzer.h"
+#include "TraceServices/ModuleService.h"
+
+#define LOCTEXT_NAMESPACE "HttpInsightsModule"
+
+DEFINE_LOG_CATEGORY(LogHttpInsights);
+
+namespace UE::HttpInsights
+{
+
+////////////////////////////////////////////////////////////////////////////////////////////////////
+namespace View
+{
+
+////////////////////////////////////////////////////////////////////////////////////////////////////
+void RegisterTimingProfilerExtensions(
+	FInsightsMajorTabExtender& InOutExtender,
+	Insights::Timing::ITimingViewExtender* TimingViewExtender);
+
+////////////////////////////////////////////////////////////////////////////////////////////////////
+TUniquePtr<Insights::Timing::ITimingViewExtender> MakeTimingViewExtender();
+
+}
+
+////////////////////////////////////////////////////////////////////////////////////////////////////
+TUniquePtr<UE::Trace::IAnalyzer> MakeHttpAnalyzer(
+	TraceServices::IAnalysisSession& Session,
+	class IHttpLogModel& LogModel);
+
+////////////////////////////////////////////////////////////////////////////////////////////////////
+class FTraceModule
+	: public TraceServices::IModule
+{
+	static const FName		ModuleName;
+
+public:
+	virtual					~FTraceModule() = default;
+	virtual void			GetModuleInfo(TraceServices::FModuleInfo& OutModuleInfo) override;
+	virtual void			GetLoggers(TArray<const TCHAR *>& OutLoggers) override;
+	virtual bool			ShouldBeEnabledByDefault() const override { return true; }
+	virtual const TCHAR*	GetCommandLineArgument() override { return TEXT("http"); }
+	virtual void			OnAnalysisBegin(TraceServices::IAnalysisSession& Session) override;
+	virtual void			GenerateReports(
+								const TraceServices::IAnalysisSession& Session,
+								const TCHAR* CmdLine,
+								const TCHAR* OutputDirectory) override;
+};
+
+const FName FTraceModule::ModuleName = FName("TraceModule_Http");
+
+////////////////////////////////////////////////////////////////////////////////////////////////////
+void FTraceModule::GetModuleInfo(TraceServices::FModuleInfo& OutModuleInfo)
+{
+	OutModuleInfo.Name = ModuleName;
+	OutModuleInfo.DisplayName = TEXT("Http");
+}
+
+void FTraceModule::GetLoggers(TArray<const TCHAR *>& OutLoggers)
+{
+	OutLoggers.Add(TEXT("Http"));
+}
+
+void FTraceModule::OnAnalysisBegin(TraceServices::IAnalysisSession& Session)
+{
+	TSharedPtr<IHttpLogModel> LogModel = MakeHttpLogModel(Session);
+	TUniquePtr<UE::Trace::IAnalyzer> Analyzer = MakeHttpAnalyzer(Session, *LogModel);
+
+	Session.AddProvider(IHttpLogModel::ProviderName, LogModel);
+	Session.AddAnalyzer(Analyzer.Release());
+}
+
+void FTraceModule::GenerateReports(
+	const TraceServices::IAnalysisSession& Session,
+	const TCHAR* CmdLine,
+	const TCHAR* OutputDirectory)
+{
+}
+
+////////////////////////////////////////////////////////////////////////////////////////////////////
+class FHttpInsightsModule
+	: public IHttpInsightsModule
+{
+public:
+	virtual void	StartupModule() override;
+	virtual void	ShutdownModule() override;
+
+private:
+	void			RegisterTimingProfilerExtensions(FInsightsMajorTabExtender& InOutExtender);
+
+	TUniquePtr<FTraceModule>							TraceModule;
+	TUniquePtr<Insights::Timing::ITimingViewExtender>	TimingViewExtender;
+};
+
+////////////////////////////////////////////////////////////////////////////////////////////////////
+void FHttpInsightsModule::StartupModule()
+{
+	TraceModule = MakeUnique<FTraceModule>();
+	IModularFeatures::Get().RegisterModularFeature(TraceServices::ModuleFeatureName, TraceModule.Get());
+	TimingViewExtender = View::MakeTimingViewExtender();
+	IModularFeatures::Get().RegisterModularFeature(Insights::Timing::TimingViewExtenderFeatureName, TimingViewExtender.Get());
+
+	IUnrealInsightsModule& Insights = FModuleManager::LoadModuleChecked<IUnrealInsightsModule>("TraceInsights");
+	Insights.OnRegisterMajorTabExtension(FInsightsManagerTabs::TimingProfilerTabId)
+		.AddRaw(this, &FHttpInsightsModule::RegisterTimingProfilerExtensions);
+}
+
+void FHttpInsightsModule::ShutdownModule()
+{
+	if (IUnrealInsightsModule* Insights = FModuleManager::GetModulePtr<IUnrealInsightsModule>("TraceInsights"))
+	{
+		Insights->OnRegisterMajorTabExtension(FInsightsManagerTabs::TimingProfilerTabId)
+			.RemoveAll(this);
+	}
+
+	IModularFeatures::Get().UnregisterModularFeature(TraceServices::ModuleFeatureName, TraceModule.Get());
+	IModularFeatures::Get().UnregisterModularFeature(Insights::Timing::TimingViewExtenderFeatureName, TimingViewExtender.Get());
+}
+
+void FHttpInsightsModule::RegisterTimingProfilerExtensions(FInsightsMajorTabExtender& InOutExtender)
+{
+	View::RegisterTimingProfilerExtensions(InOutExtender, TimingViewExtender.Get());
+}
+
+} // namespace UE:HttpInsights
+
+IMPLEMENT_MODULE(UE::HttpInsights::FHttpInsightsModule, HttpInsights)
+
+#undef LOCTEXT_NAMESPACE
diff --git a/Engine/Plugins/Experimental/HttpInsights/Source/Private/View.cpp b/Engine/Plugins/Experimental/HttpInsights/Source/Private/View.cpp
new file mode 100644
--- /dev/null
+++ b/Engine/Plugins/Experimental/HttpInsights/Source/Private/View.cpp
@@ -0,0 +1,778 @@
+// Copyright Epic Games, Inc. All Rights Reserved.
+
+#include "Model.h"
+
+#include "Delegates/Delegate.h"
+#include "Framework/Commands/UICommandList.h"
+#include "Framework/MultiBox/MultiBoxBuilder.h"
+#include "Insights/IUnrealInsightsModule.h"
+#include "Insights/ITimingViewExtender.h"
+#include "Insights/ITimingViewSession.h"
+#include "Math/Color.h"
+#include "Misc/EnumClassFlags.h"
+#include "Widgets/DeclarativeSyntaxSupport.h"
+#include "Widgets/Docking/SDockTab.h"
+#include "Widgets/Input/SEditableTextBox.h"
+#include "Widgets/SBoxPanel.h"
+#include "Widgets/SCompoundWidget.h"
+#include "Widgets/Text/STextBlock.h"
+#include "Widgets/Views/SHeaderRow.h"
+#include "Widgets/Views/SListView.h"
+#include "Widgets/Text/STextBlock.h"
+
+#define LOCTEXT_NAMESPACE "HttpInsightsView"
+
+namespace UE::HttpInsights::View
+{
+
+using namespace Insights::Timing;
+using namespace TraceServices;
+
+////////////////////////////////////////////////////////////////////////////////////////////////////
+class FHttpTimingViewSession
+{
+public:
+							FHttpTimingViewSession() = default;
+	bool					IsValid() const { return TimingSession != nullptr && AnalysisSession != nullptr; }
+	ITimingViewSession&		GetTimingViewSession() const { return *TimingSession; }
+	const IAnalysisSession&	GetAnalysisSession() const { return *AnalysisSession; }
+
+private:
+	friend class FHttpTimingViewExtender;
+
+	ITimingViewSession*		TimingSession = nullptr;
+	const IAnalysisSession*	AnalysisSession = nullptr;
+};
+
+////////////////////////////////////////////////////////////////////////////////////////////////////
+class IHttpTimingView
+{
+public:
+	virtual void OnBeginSession(FHttpTimingViewSession& Session) = 0;
+	virtual void OnEndSession(FHttpTimingViewSession& Session) = 0;
+	virtual void OnTickSession(FHttpTimingViewSession&) = 0;
+};
+
+////////////////////////////////////////////////////////////////////////////////////////////////////
+class FHttpTimingViewExtender
+	: public Insights::Timing::ITimingViewExtender
+{
+public:
+	virtual void	OnBeginSession(ITimingViewSession& TimingViewSession) override;
+	virtual void	OnEndSession(ITimingViewSession& TimingViewSession) override;
+	virtual void	Tick(ITimingViewSession& TimingViewSession, const IAnalysisSession& AnalysisSession) override;
+	void			AddView(IHttpTimingView* View);
+	void			RemoveView(IHttpTimingView* View);
+
+private:
+	FHttpTimingViewSession		CurrentSession;
+	TArray<IHttpTimingView*>	Views;
+};
+
+////////////////////////////////////////////////////////////////////////////////////////////////////
+void FHttpTimingViewExtender::OnBeginSession(ITimingViewSession& TimingViewSession)
+{
+	if (TimingViewSession.GetName() != FInsightsManagerTabs::TimingProfilerTabId)
+	{
+		return;
+	}
+
+	check(CurrentSession.TimingSession == nullptr);
+	check(CurrentSession.AnalysisSession == nullptr);
+	CurrentSession.TimingSession = &TimingViewSession;
+}
+
+void FHttpTimingViewExtender::OnEndSession(ITimingViewSession& TimingViewSession)
+{
+	if (TimingViewSession.GetName() != FInsightsManagerTabs::TimingProfilerTabId)
+	{
+		return;
+	}
+
+	if (CurrentSession.TimingSession != &TimingViewSession) 
+	{
+		return;
+	}
+
+	if (CurrentSession.IsValid())
+	{
+		for (IHttpTimingView* View : Views)
+		{
+			View->OnEndSession(CurrentSession);
+		}
+	}
+	
+	CurrentSession = FHttpTimingViewSession();
+}
+
+void FHttpTimingViewExtender::Tick(ITimingViewSession& TimingViewSession, const IAnalysisSession& AnalysisSession)
+{
+	if (TimingViewSession.GetName() != FInsightsManagerTabs::TimingProfilerTabId)
+	{
+		return;
+	}
+
+	if (CurrentSession.TimingSession != &TimingViewSession)
+	{
+		return;
+	}
+
+	if (CurrentSession.AnalysisSession == nullptr)
+	{
+		CurrentSession.AnalysisSession = &AnalysisSession;
+		for (IHttpTimingView* View : Views)
+		{
+			View->OnBeginSession(CurrentSession);
+		}
+	}
+
+	for (IHttpTimingView* View : Views)
+	{
+		View->OnTickSession(CurrentSession);
+	}
+}
+
+void FHttpTimingViewExtender::AddView(IHttpTimingView* View)
+{
+	Views.Add(View);
+	if (CurrentSession.IsValid())
+	{
+		View->OnBeginSession(CurrentSession);
+	}
+}
+
+void FHttpTimingViewExtender::RemoveView(IHttpTimingView* View)
+{
+	Views.Remove(View);
+}
+
+////////////////////////////////////////////////////////////////////////////////////////////////////
+struct FHttpRequestViewModelData
+{
+	FHttpRequestViewModelData(const IHttpLogModel& InModel, const FHttpRequest& InRequest)
+		: Model(InModel)
+		, Request(InRequest)
+		, FullUrl(InRequest.Host / InRequest.Url)
+	{
+		StatusColor = Request.StatusCode > 199 && Request.StatusCode < 300
+			? FLinearColor::Green
+			: FLinearColor::Red;
+
+		// Just some simple colors for IAS/IAD
+		CategoryColor = InRequest.Category == 0
+			? FLinearColor(1.0f, 0.0f, 1.0f)
+			: InRequest.Category == 1
+				? FLinearColor(0.0f, 1.0f, 1.0f)
+				: FLinearColor::Gray;
+	}
+
+	bool	HasChunkRanges() const { return Request.ChunkRanges != nullptr; }
+	bool	MatchesText(const FString& SearchText);
+	FText	GetChunkRangeText(const FHttpChunkRange& ChunkRange) const;
+
+	const IHttpLogModel&	Model;
+	const FHttpRequest&		Request;
+	const FString			FullUrl;
+	FLinearColor			StatusColor;
+	FLinearColor			CategoryColor;
+};
+using FHttpRequestViewModel = TSharedPtr<FHttpRequestViewModelData>;
+
+////////////////////////////////////////////////////////////////////////////////////////////////////
+bool FHttpRequestViewModelData::MatchesText(const FString& SearchText)
+{
+	for (FHttpChunkRange* It = Request.ChunkRanges; It != nullptr; It = It->Next)
+	{
+		if (const TCHAR* PackageName = Model.GetPackageName(It->ChunkId))
+		{
+			//TODO: Fix
+			FString Tmp(PackageName);
+			if (Tmp.Contains(SearchText))
+			{
+				return true;
+			}
+		}
+	}
+
+	return false;
+}
+
+FText FHttpRequestViewModelData::GetChunkRangeText(const FHttpChunkRange& ChunkRange) const
+{
+	if (const TCHAR* PackageName = Model.GetPackageName(ChunkRange.ChunkId))
+	{
+		return FText::FromString(
+			FString::Printf(TEXT("%s [%s] [%u - %u] %s"),
+				PackageName,
+				*LexToString(ChunkRange.ChunkId.GetChunkType()),
+				ChunkRange.Start, ChunkRange.End,
+				*FText::AsMemory(ChunkRange.End - ChunkRange.Start).ToString()));
+	}
+	else
+	{
+		return FText::FromString(
+			FString::Printf(TEXT("%s [%s] [%u - %u] %s"),
+				*LexToString(ChunkRange.ChunkId),
+				*LexToString(ChunkRange.ChunkId.GetChunkType()),
+				ChunkRange.Start, ChunkRange.End,
+				*FText::AsMemory(ChunkRange.End - ChunkRange.Start).ToString()));
+	}
+}
+
+////////////////////////////////////////////////////////////////////////////////////////////////////
+class SHttpRequestTableRow
+	: public SMultiColumnTableRow<FHttpRequestViewModel>
+{
+public:
+	using Super = SMultiColumnTableRow<FHttpRequestViewModel>;
+
+	inline static const FName TimeColumn			= FName("Time");
+	inline static const FName StatusColumn			= FName("Status");
+	inline static const FName DurationColumn		= FName("Duration");
+	inline static const FName ContentLengthColumn	= FName("ContentLength");
+	inline static const FName UrlColumn				= FName("Url");
+	inline static const FName CategoryColumn		= FName("Category");
+
+	SLATE_BEGIN_ARGS(SHttpRequestTableRow) { }
+	SLATE_END_ARGS()
+
+	void Construct(
+		const FArguments& Args,
+		const TSharedRef<STableViewBase>& OwnerTable,
+		FHttpRequestViewModel InViewModel)
+	{
+		ViewModel = InViewModel;
+		Super::Construct(Super::FArguments().Padding(1.0f), OwnerTable);
+	}
+
+	static TSharedPtr<SHeaderRow>	CreateHeaderRow();
+private:
+	TSharedRef<SWidget>				GenerateWidgetForColumn(const FName& Column);
+
+	FHttpRequestViewModel ViewModel;
+};
+
+TSharedPtr<SHeaderRow> SHttpRequestTableRow::CreateHeaderRow()
+{
+	return SNew(SHeaderRow)
+		+SHeaderRow::Column(TimeColumn)
+		.DefaultLabel(FText::FromString("Time"))
+		.HAlignHeader(EHorizontalAlignment::HAlign_Center)
+		.FixedWidth(64)
+		+SHeaderRow::Column(StatusColumn)
+		.DefaultLabel(LOCTEXT("StatusColumnName", "Status"))
+		.HAlignHeader(EHorizontalAlignment::HAlign_Center)
+		.FixedWidth(64)
+		+SHeaderRow::Column(DurationColumn)
+		.DefaultLabel(LOCTEXT("DurationColumnName", "Duration"))
+		.HAlignHeader(EHorizontalAlignment::HAlign_Center)
+		.FixedWidth(72)
+		+SHeaderRow::Column(ContentLengthColumn)
+		.DefaultLabel(LOCTEXT("ContentLengthColumnName", "Content-Length"))
+		.HAlignHeader(EHorizontalAlignment::HAlign_Left)
+		.FixedWidth(128)
+		+SHeaderRow::Column(UrlColumn)
+		.DefaultLabel(LOCTEXT("UrlColumnName", "Url"))
+		.HAlignHeader(EHorizontalAlignment::HAlign_Left)
+		+SHeaderRow::Column(CategoryColumn)
+		.DefaultLabel(LOCTEXT("CategoryColumnName", "Category"))
+		.HAlignHeader(EHorizontalAlignment::HAlign_Center);
+}
+
+////////////////////////////////////////////////////////////////////////////////////////////////////
+TSharedRef<SWidget> SHttpRequestTableRow::GenerateWidgetForColumn(const FName& Column)
+{
+	const float PaddingLeft = 12.0f;
+
+	if (Column == TimeColumn)
+	{
+		return SNew(SBox)
+			.HAlign(HAlign_Center)
+			[
+				SNew(STextBlock)
+				.Text(FText::AsTimespan(FTimespan::FromSeconds(ViewModel->Request.CompletionTime)))
+			];
+	}
+	else if (Column == StatusColumn)
+	{
+		return SNew(SBox)
+			.HAlign(HAlign_Center)
+			[
+				SNew(STextBlock)
+				.Text(FText::AsNumber(ViewModel->Request.StatusCode))
+				.ColorAndOpacity(ViewModel->StatusColor)
+			];
+	}
+	else if (Column == DurationColumn)
+	{
+		return SNew(SBox)
+			.HAlign(HAlign_Left)
+			.Padding(PaddingLeft, 0, 0, 0)
+			[
+				SNew(STextBlock)
+				.Text(FText::AsNumber(ViewModel->Request.CompletionTime - ViewModel->Request.StartTime))
+			];
+	}
+	else if (Column == ContentLengthColumn)
+	{
+		return SNew(SBox)
+			.HAlign(HAlign_Left)
+			.Padding(PaddingLeft, 0, 0, 0)
+			[
+				SNew(STextBlock)
+					.Text(FText::AsMemory(ViewModel->Request.ContentLength))
+			];
+	}
+	else if (Column == UrlColumn)
+	{
+		if (ViewModel->HasChunkRanges())
+		{
+			TSharedRef<SVerticalBox> Vertical = SNew(SVerticalBox)
+				+SVerticalBox::Slot()
+				[
+					SNew(STextBlock)
+					.Text(FText::FromString(ViewModel->FullUrl))
+				];
+
+			for (FHttpChunkRange* It = ViewModel->Request.ChunkRanges; It != nullptr; It = It->Next)
+			{
+				Vertical->AddSlot()
+				[
+					SNew(STextBlock)
+					.Text(ViewModel->GetChunkRangeText(*It))
+				];
+			}
+
+			return SNew(SBox).Padding(PaddingLeft, 0, 0, 0)[Vertical];
+		}
+
+		return SNew(SBox)
+			.Padding(PaddingLeft, 0, 0, 0)
+			[
+				SNew(STextBlock)
+				.Text(FText::FromString(ViewModel->FullUrl))
+			];
+	}
+	else if (Column == CategoryColumn)
+	{
+		return SNew(SBox)
+			.Padding(PaddingLeft, 0, 0, 0)
+			[
+				SNew(STextBlock)
+				.Text(FText::FromString(*ViewModel->Request.CategoryName))
+				.ColorAndOpacity(ViewModel->CategoryColor)
+			];
+	}
+
+	return SNullWidget::NullWidget;
+}
+
+////////////////////////////////////////////////////////////////////////////////////////////////////
+class SHttpRequestTableView
+	: public SListView<FHttpRequestViewModel>
+{
+public:
+	using Super				= SListView<FHttpRequestViewModel>;
+
+							SLATE_BEGIN_ARGS(SHttpRequestTableView) { }
+								SLATE_ARGUMENT(TArray<FHttpRequestViewModel>*, ListItemsSource)
+							SLATE_END_ARGS()
+	
+	void					Construct(const FArguments& Args);
+	void					Refresh();
+
+private:
+	virtual FReply			OnMouseWheel(const FGeometry& InGeometry, const FPointerEvent& InMouseEvent) override;
+	virtual FReply			OnKeyDown(const FGeometry& InGeometry, const FKeyEvent& InKeyEvent) override;
+	void					HandleUserScrolled(double Offset);
+	TSharedRef<ITableRow>	OnGenerateRow(FHttpRequestViewModel ViewModel, const TSharedRef<STableViewBase>& OwnerTable)
+							{
+								return SNew(SHttpRequestTableRow, OwnerTable, ViewModel);
+							}
+
+	bool bAutoScroll = true;
+};
+
+////////////////////////////////////////////////////////////////////////////////////////////////////
+void SHttpRequestTableView::Construct(const FArguments& InArgs)
+{
+	Super::Construct
+	(
+		Super::FArguments()
+		.ScrollbarVisibility(EVisibility::Visible)
+		.ListItemsSource(InArgs._ListItemsSource)
+		.OnGenerateRow(this, &SHttpRequestTableView::OnGenerateRow)
+		.HeaderRow(SHttpRequestTableRow::CreateHeaderRow())
+		.OnListViewScrolled(this, &SHttpRequestTableView::HandleUserScrolled)
+	);
+}
+
+void SHttpRequestTableView::Refresh()
+{
+	Super::RequestListRefresh();
+	if (bAutoScroll)
+	{
+		Super::ScrollToBottom();
+	}
+}
+
+FReply SHttpRequestTableView::OnMouseWheel(const FGeometry& InGeometry, const FPointerEvent& InMouseEvent)
+{
+	bAutoScroll = false;
+	return Super::OnMouseWheel(InGeometry, InMouseEvent);
+}
+
+FReply SHttpRequestTableView::OnKeyDown(const FGeometry& InGeometry, const FKeyEvent& InKeyEvent)
+{
+	if (bAutoScroll == false)
+	{
+		bAutoScroll = InKeyEvent.GetKey() == EKeys::PageDown;
+	}
+	return Super::OnKeyDown(InGeometry, InKeyEvent);
+}
+
+void SHttpRequestTableView::HandleUserScrolled(double Offset)
+{
+	if (bAutoScroll && IsUserScrolling())
+	{
+		bAutoScroll = false;
+	}
+}
+
+////////////////////////////////////////////////////////////////////////////////////////////////////
+enum class EHttpCategoryFilter : uint8
+{
+	None		= 0,
+	Streaming	= (1 << 0),
+	Install		= (1 << 1),
+	All			= Streaming | Install
+};
+ENUM_CLASS_FLAGS(EHttpCategoryFilter)
+
+////////////////////////////////////////////////////////////////////////////////////////////////////
+class SHttpFilterBar
+	: public SCompoundWidget
+{
+public:
+	DECLARE_MULTICAST_DELEGATE(FFilterChangedDelegate);
+
+								SLATE_BEGIN_ARGS(SHttpFilterBar) { }
+								SLATE_END_ARGS()
+	void						Construct(const FArguments& Args);
+	bool						ShouldFilter(const FHttpRequestViewModel& Vm) const;
+	FFilterChangedDelegate&		OnFilterChanged() { return FilterChanged; }
+
+private:
+	void						UpdateCategoryFilter(EHttpCategoryFilter Filter);
+	bool						IsCategoryFilterEnabled(EHttpCategoryFilter Filer);
+	void						HandleSearchTextChanged(const FText& NewSearchText);
+	TSharedRef<SWidget>			MakeCategoryMenu();
+
+	TSharedPtr<SEditableTextBox>	SearchTextBox;
+	FFilterChangedDelegate			FilterChanged;
+	EHttpCategoryFilter				CurrentCategoryFilter = EHttpCategoryFilter::All;
+};
+
+////////////////////////////////////////////////////////////////////////////////////////////////////
+void SHttpFilterBar::Construct(const FArguments& Args)
+{
+	FSlimHorizontalToolBarBuilder ToolbarBuilder(TSharedPtr<const FUICommandList>(), FMultiBoxCustomization::None);
+	ToolbarBuilder.BeginSection("Filters");
+	{
+		ToolbarBuilder.AddComboButton(
+			FUIAction(),
+			FOnGetContent::CreateSP(this, &SHttpFilterBar::MakeCategoryMenu),
+			LOCTEXT("HttpCategoryText", "Category Filter"),
+			LOCTEXT("HttpThresholdToolTip", "Filter HTTP request by category."),
+			FSlateIcon(),
+			false
+		);
+	}
+	ToolbarBuilder.EndSection();
+
+	ChildSlot
+	[
+		SNew(SHorizontalBox)
+		+SHorizontalBox::Slot()
+		.AutoWidth()
+		[
+			ToolbarBuilder.MakeWidget()
+		]
+		+SHorizontalBox::Slot()
+		.FillWidth(1.0f)
+		[
+			SNew(SBox)
+			.VAlign(EVerticalAlignment::VAlign_Center)
+			.Padding(FMargin(11.0f, 0.0f))
+			[
+				SAssignNew(SearchTextBox, SEditableTextBox)
+				.HintText(NSLOCTEXT("HttpRequestSearchBox", "Search", "Search..."))
+				.SelectAllTextWhenFocused(true)
+				.OnTextChanged(this, &SHttpFilterBar::HandleSearchTextChanged)
+			]
+		]
+	];
+}
+
+bool SHttpFilterBar::ShouldFilter(const FHttpRequestViewModel& Vm) const
+{
+	//TODO: Category filtering is currently hardcoded see OnDemandHttpIoDispatcher.h
+	if (!EnumHasAnyFlags(CurrentCategoryFilter, EHttpCategoryFilter::Streaming) && Vm->Request.Category == 0)
+	{
+		return true;
+	}
+
+	if (!EnumHasAnyFlags(CurrentCategoryFilter, EHttpCategoryFilter::Install) && Vm->Request.Category == 1)
+	{
+		return true;
+	}
+
+	const FText SearchText = SearchTextBox->GetText();
+	if (SearchText.IsEmpty() == false)
+	{
+		const FString Text = SearchText.ToString();
+		if (Vm->MatchesText(Text) == false)
+		{
+			return true;
+		}
+	}
+
+	return false;
+}
+
+TSharedRef<SWidget>	SHttpFilterBar::MakeCategoryMenu()
+{
+	struct FCategoryFilter
+	{
+		EHttpCategoryFilter	Value;
+		FText				Label;
+		FText				ToolTip;
+	};
+
+	const FCategoryFilter Filters[]
+	{
+		{
+			EHttpCategoryFilter::Streaming,
+			LOCTEXT("HttpCategoryMenuStreaming", "Streaming"),
+			LOCTEXT("HttpCategoryStreaming_Tooltip", "Filter HTTP requests by Streaming category.")
+		},
+		{
+			EHttpCategoryFilter::Install,
+			LOCTEXT("HttpCategoryMenuInstall", "Install"),
+			LOCTEXT("HttpCategoryInstall_Tooltip", "Filter HTTP requests by Install category.")
+		}
+	};
+
+	FMenuBuilder MenuBuilder(true, nullptr, nullptr, false, &FCoreStyle::Get(), false);
+	MenuBuilder.BeginSection("HttpCategory");
+	for (int32 Idx = 0, Count = UE_ARRAY_COUNT(Filters); Idx < Count; ++Idx)
+	{
+		const FCategoryFilter& Filter = Filters[Idx];
+		const TSharedRef<SWidget> TextBlock = 
+			SNew(STextBlock)
+			.Text(Filter.Label);
+
+		MenuBuilder.AddMenuEntry(
+			FUIAction(
+				FExecuteAction::CreateSP(this, &SHttpFilterBar::UpdateCategoryFilter, Filter.Value),
+				FCanExecuteAction(),
+				FIsActionChecked::CreateSP(this, &SHttpFilterBar::IsCategoryFilterEnabled, Filter.Value)),
+			TextBlock,
+			NAME_None,
+			Filter.ToolTip,
+			EUserInterfaceActionType::Check
+		);
+	}
+	MenuBuilder.EndSection();
+
+	return MenuBuilder.MakeWidget();
+}
+
+void SHttpFilterBar::UpdateCategoryFilter(EHttpCategoryFilter Filter)
+{
+	if (EnumHasAnyFlags(CurrentCategoryFilter, Filter))
+	{
+		EnumRemoveFlags(CurrentCategoryFilter, Filter);
+	}
+	else
+	{
+		EnumAddFlags(CurrentCategoryFilter, Filter);
+	}
+
+	FilterChanged.Broadcast();
+}
+
+bool SHttpFilterBar::IsCategoryFilterEnabled(EHttpCategoryFilter Filter)
+{
+	return EnumHasAnyFlags(CurrentCategoryFilter, Filter);
+}
+
+void SHttpFilterBar::HandleSearchTextChanged(const FText& NewSearchText)
+{
+	FilterChanged.Broadcast();
+}
+
+////////////////////////////////////////////////////////////////////////////////////////////////////
+class SHttpInsightsTab
+	: public SCompoundWidget
+	, public IHttpTimingView
+{
+public:
+	static const FName		TabId;
+
+							SLATE_BEGIN_ARGS(SHttpInsightsTab) { }
+								SLATE_ARGUMENT(FHttpTimingViewExtender*, TimingViewExtender)
+							SLATE_END_ARGS()
+							~SHttpInsightsTab();
+	void					Construct(const FArguments& InArgs);
+
+	virtual void			OnBeginSession(FHttpTimingViewSession& Session);
+	virtual void			OnEndSession(FHttpTimingViewSession& Session);
+	virtual void			OnTickSession(FHttpTimingViewSession& Session) { }
+	virtual void			Tick(const FGeometry& AllottedGeometry, const double InCurrentTime, const float InDeltaTime) override;
+
+private:
+	FHttpTimingViewExtender*			TimingViewExtender = nullptr;
+	FHttpTimingViewSession*				CurrentSession = nullptr;
+	TArray<FHttpRequestViewModel>		HttpRequests;
+	TArray<FHttpRequestViewModel>		AllHttpRequests;
+	TSharedPtr<SHttpFilterBar>			FilterBar;
+	TSharedPtr<SHttpRequestTableView>	RequestTable;
+	bool								bFilterChanged = false;
+};
+
+////////////////////////////////////////////////////////////////////////////////////////////////////
+const FName SHttpInsightsTab::TabId = FName("HttpInsightsTab");
+
+SHttpInsightsTab::~SHttpInsightsTab()
+{
+	if (TimingViewExtender != nullptr)
+	{
+		TimingViewExtender->RemoveView(this);
+	}
+}
+
+void SHttpInsightsTab::Construct(const FArguments& InArgs)
+{
+	TimingViewExtender = InArgs._TimingViewExtender;
+	check(TimingViewExtender != nullptr);
+
+	ChildSlot
+	[
+		SNew(SVerticalBox)
+		+SVerticalBox::Slot()
+		.AutoHeight()
+		[
+			SAssignNew(FilterBar, SHttpFilterBar)
+		]
+		+SVerticalBox::Slot()
+		[
+			SAssignNew(RequestTable, SHttpRequestTableView)
+			.ListItemsSource(&HttpRequests)
+		]
+	];
+
+	TimingViewExtender->AddView(this);
+	FilterBar->OnFilterChanged().AddLambda([this]
+	{
+		bFilterChanged = true;
+	});
+}
+
+void SHttpInsightsTab::OnBeginSession(FHttpTimingViewSession& Session)
+{
+	check(CurrentSession == nullptr);
+	CurrentSession = &Session;
+}
+
+void SHttpInsightsTab::OnEndSession(FHttpTimingViewSession& Session)
+{
+	if (CurrentSession == &Session)
+	{
+		CurrentSession = nullptr;
+		HttpRequests.Empty();
+		AllHttpRequests.Empty();
+	}
+}
+
+void SHttpInsightsTab::Tick(const FGeometry& AllottedGeometry, const double InCurrentTime, const float InDeltaTime)
+{
+	if (CurrentSession == nullptr)
+	{
+		HttpRequests.Empty();
+		return;
+	}
+
+	const IHttpLogModel* HttpLog = CurrentSession->GetAnalysisSession().ReadProvider<IHttpLogModel>(IHttpLogModel::ProviderName);
+	if (HttpLog == nullptr)
+	{
+		HttpRequests.Empty();
+		return;
+	}
+
+	const int32 Num		= AllHttpRequests.Num();
+	const int32 NewNum	= HttpLog->IterateLog(Num, [this, HttpLog](const FHttpRequest& Request)
+	{
+		FHttpRequestViewModel Vm = MakeShared<FHttpRequestViewModelData>(*HttpLog, Request);
+		AllHttpRequests.Add(Vm);
+
+		if (!bFilterChanged && FilterBar->ShouldFilter(Vm) == false)
+		{
+			HttpRequests.Add(Vm);
+		}
+	});
+
+	const bool bRefresh = bFilterChanged || (Num != NewNum);
+
+	if (bFilterChanged)
+	{
+		bFilterChanged = false;
+		HttpRequests.Empty(AllHttpRequests.Num());
+		for (const FHttpRequestViewModel& Request : AllHttpRequests)
+		{
+			if (FilterBar->ShouldFilter(Request) == false)
+			{
+				HttpRequests.Add(Request);
+			}
+		}
+	}
+
+	if (bRefresh)
+	{
+		RequestTable->Refresh();
+	}
+}
+
+////////////////////////////////////////////////////////////////////////////////////////////////////
+void RegisterTimingProfilerExtensions(
+	FInsightsMajorTabExtender& InOutExtender,
+	Insights::Timing::ITimingViewExtender* TimingViewExtender)
+{
+	FMinorTabConfig& MinorTabConfig = InOutExtender.AddMinorTabConfig();
+
+	MinorTabConfig.TabId			= SHttpInsightsTab::TabId;
+	MinorTabConfig.TabLabel			= LOCTEXT("HttpInsightsTabTitle", "HTTP Insights");
+	MinorTabConfig.TabTooltip		= LOCTEXT("HttpInsightsTabTitleTooltip", "View and diagnose HTTP traffic.");
+	MinorTabConfig.TabIcon			= FSlateIcon(FAppStyle::Get().GetStyleSetName(), "Icons.Plugin.TreeItem");
+	MinorTabConfig.WorkspaceGroup	= InOutExtender.GetWorkspaceGroup();
+	MinorTabConfig.OnSpawnTab		= FOnSpawnTab::CreateLambda([TimingViewExtender](const FSpawnTabArgs& Args)
+	{
+		return SNew(SDockTab)
+			.ShouldAutosize(false)
+			.TabRole(ETabRole::PanelTab)
+			[
+				SNew(SHttpInsightsTab)
+				.TimingViewExtender(reinterpret_cast<FHttpTimingViewExtender*>(TimingViewExtender))
+			];
+	});
+}
+
+////////////////////////////////////////////////////////////////////////////////////////////////////
+TUniquePtr<Insights::Timing::ITimingViewExtender> MakeTimingViewExtender()
+{
+	return MakeUnique<FHttpTimingViewExtender>();
+}
+
+} // namespace UE::HttpInsights::View
+
+#undef LOCTEXT_NAMESPACE
diff --git a/Engine/Plugins/Experimental/HttpInsights/Source/Public/HttpInsights/ModuleInterface.h b/Engine/Plugins/Experimental/HttpInsights/Source/Public/HttpInsights/ModuleInterface.h
new file mode 100644
--- /dev/null
+++ b/Engine/Plugins/Experimental/HttpInsights/Source/Public/HttpInsights/ModuleInterface.h
@@ -0,0 +1,25 @@
+// Copyright Epic Games, Inc. All Rights Reserved.
+
+#pragma once 
+
+#include "Modules/ModuleInterface.h"
+#include "Logging/LogMacros.h"
+#include "Logging/LogVerbosity.h"
+
+#define UE_API HTTPINSIGHTS_API
+
+UE_API DECLARE_LOG_CATEGORY_EXTERN(LogHttpInsights, Log, All);
+
+namespace UE::HttpInsights
+{
+
+////////////////////////////////////////////////////////////////////////////////////////////////////
+class IHttpInsightsModule
+	: public IModuleInterface
+{
+public:
+};
+
+} // namespace UE:HttpInsights
+
+#undef UE_API
diff --git a/Engine/Programs/UnrealInsights/Config/DefaultEngine.ini b/Engine/Programs/UnrealInsights/Config/DefaultEngine.ini
--- a/Engine/Programs/UnrealInsights/Config/DefaultEngine.ini
+++ b/Engine/Programs/UnrealInsights/Config/DefaultEngine.ini
@@ -12,6 +12,7 @@
 +ProgramEnabledPlugins=MassInsights
 +ProgramEnabledPlugins=ChaosInsights
 +ProgramEnabledPlugins=MessagingInsights
++ProgramEnabledPlugins=HttpInsights
 
 [/Script/MacTargetPlatform.XcodeProjectSettings]
 TemplateMacPlist=(FilePath="/Engine/Source/Programs/UnrealInsights/Resources/Mac/Info.plist")
diff --git a/Engine/Source/Runtime/Core/Private/ProfilingDebugging/HttpTrace.cpp b/Engine/Source/Runtime/Core/Private/ProfilingDebugging/HttpTrace.cpp
new file mode 100644
--- /dev/null
+++ b/Engine/Source/Runtime/Core/Private/ProfilingDebugging/HttpTrace.cpp
@@ -0,0 +1,97 @@
+// Copyright Epic Games, Inc. All Rights Reserved.
+
+#include "ProfilingDebugging/HttpTrace.h"
+
+#if UE_TRACE_HTTP_ENABLED
+
+#include "HAL/PlatformTime.h"
+#include "IO/IoChunkId.h"
+#include "Trace/Trace.h"
+
+////////////////////////////////////////////////////////////////////////////////////////////////////
+UE_TRACE_CHANNEL_DEFINE(HttpChannel)
+
+////////////////////////////////////////////////////////////////////////////////////////////////////
+UE_TRACE_EVENT_BEGIN(Http, DispatcherCreated, NoSync)
+	UE_TRACE_EVENT_FIELD(uint64, Dispatcher)
+	UE_TRACE_EVENT_FIELD(UE::Trace::AnsiString, Name)
+UE_TRACE_EVENT_END()
+
+void FHttpTrace::DispatcherCreated(uint64 Dispatcher, const ANSICHAR* Name)
+{
+	UE_TRACE_LOG(Http, DispatcherCreated, HttpChannel)
+		<< DispatcherCreated.Dispatcher(Dispatcher)
+		<< DispatcherCreated.Name(Name);
+}
+
+////////////////////////////////////////////////////////////////////////////////////////////////////
+UE_TRACE_EVENT_BEGIN(Http, CategoryCreated, NoSync)
+	UE_TRACE_EVENT_FIELD(uint8, Category)
+	UE_TRACE_EVENT_FIELD(UE::Trace::AnsiString, Name)
+UE_TRACE_EVENT_END()
+
+void FHttpTrace::CategoryCreated(uint8 Category, const ANSICHAR* Name)
+{
+	UE_TRACE_LOG(Http, CategoryCreated, HttpChannel)
+		<< CategoryCreated.Category(Category)
+		<< CategoryCreated.Name(Name);
+}
+
+////////////////////////////////////////////////////////////////////////////////////////////////////
+UE_TRACE_EVENT_BEGIN(Http, RequestStarted, NoSync)
+	UE_TRACE_EVENT_FIELD(uint64, Dispatcher)
+	UE_TRACE_EVENT_FIELD(uint64, Request)
+	UE_TRACE_EVENT_FIELD(uint64, StartTime)
+	UE_TRACE_EVENT_FIELD(int32, Priority)
+	UE_TRACE_EVENT_FIELD(uint32, Category)
+	UE_TRACE_EVENT_FIELD(UE::Trace::AnsiString, Url)
+UE_TRACE_EVENT_END()
+
+void FHttpTrace::RequestStarted(uint64 Dispatcher, uint64 Request, const ANSICHAR* Url, int32 Priority, uint8 Category)
+{
+	UE_TRACE_LOG(Http, RequestStarted, HttpChannel)
+		<< RequestStarted.Dispatcher(Dispatcher)
+		<< RequestStarted.Request(Request)
+		<< RequestStarted.StartTime(FPlatformTime::Cycles64())
+		<< RequestStarted.Priority(Priority)
+		<< RequestStarted.Category(Category)
+		<< RequestStarted.Url(Url);
+}
+
+////////////////////////////////////////////////////////////////////////////////////////////////////
+UE_TRACE_EVENT_BEGIN(Http, ChunkRangeAdded, NoSync)
+	UE_TRACE_EVENT_FIELD(uint64, Request)
+	UE_TRACE_EVENT_FIELD(uint32, Start)
+	UE_TRACE_EVENT_FIELD(uint32, End)
+	UE_TRACE_EVENT_FIELD(uint8[], ChunkId)
+UE_TRACE_EVENT_END()
+
+void FHttpTrace::ChunkRangeAdded(uint64 Request, const FIoChunkId& ChunkId, uint32 Start, uint32 End)
+{
+	UE_TRACE_LOG(Http, ChunkRangeAdded, HttpChannel)
+		<< ChunkRangeAdded.Request(Request)
+		<< ChunkRangeAdded.Start(Start)
+		<< ChunkRangeAdded.End(End)
+		<< ChunkRangeAdded.ChunkId(reinterpret_cast<const uint8*>(&ChunkId), uint32(sizeof(FIoChunkId)));
+}
+
+////////////////////////////////////////////////////////////////////////////////////////////////////
+UE_TRACE_EVENT_BEGIN(Http, RequestCompleted, NoSync)
+	UE_TRACE_EVENT_FIELD(uint64, Request)
+	UE_TRACE_EVENT_FIELD(uint64, CompletionTime)
+	UE_TRACE_EVENT_FIELD(uint32, StatusCode)
+	UE_TRACE_EVENT_FIELD(uint32, ContentLength)
+	UE_TRACE_EVENT_FIELD(UE::Trace::AnsiString, Host)
+UE_TRACE_EVENT_END()
+
+void FHttpTrace::RequestCompleted(uint64 Request, const ANSICHAR* Host, uint32 StatusCode, uint32 ContentLength)
+{
+	UE_TRACE_LOG(Http, RequestCompleted, HttpChannel)
+		<< RequestCompleted.Request(Request)
+		<< RequestCompleted.CompletionTime(FPlatformTime::Cycles64())
+		<< RequestCompleted.StatusCode(StatusCode)
+		<< RequestCompleted.ContentLength(ContentLength)
+		<< RequestCompleted.Host(Host);
+}
+
+#endif //UE_TRACE_HTTP_ENABLED
diff --git a/Engine/Source/Runtime/Core/Public/ProfilingDebugging/HttpTrace.h b/Engine/Source/Runtime/Core/Public/ProfilingDebugging/HttpTrace.h
new file mode 100644
--- /dev/null
+++ b/Engine/Source/Runtime/Core/Public/ProfilingDebugging/HttpTrace.h
@@ -0,0 +1,67 @@
+// Copyright Epic Games, Inc. All Rights Reserved.
+#pragma once
+
+#include "CoreTypes.h"
+#include "Misc/Build.h"
+#include "ProfilingDebugging/MetadataTrace.h"
+#include "Trace/Config.h"
+
+#define UE_API CORE_API
+
+////////////////////////////////////////////////////////////////////////////////////////////////////
+#if !defined(UE_TRACE_HTTP_ENABLED) // to allow users to manually disable it from build.cs files
+	#if !UE_BUILD_SHIPPING
+		#define UE_TRACE_HTTP_ENABLED UE_TRACE_ENABLED
+	#else
+		#define UE_TRACE_HTTP_ENABLED 0
+	#endif
+#endif
+
+////////////////////////////////////////////////////////////////////////////////////////////////////
+#if !defined(UE_TRACE_HTTP_METADATA_ENABLED)
+	#define UE_TRACE_HTTP_METADATA_ENABLED UE_TRACE_HTTP_ENABLED && UE_TRACE_METADATA_ENABLED
+#endif
+
+////////////////////////////////////////////////////////////////////////////////////////////////////
+#if UE_TRACE_HTTP_ENABLED
+
+UE_TRACE_CHANNEL_EXTERN(HttpChannel, CORE_API);
+
+////////////////////////////////////////////////////////////////////////////////////////////////////
+struct FHttpTrace
+{
+	static UE_API void DispatcherCreated(uint64 Dispatcher, const ANSICHAR* Name);
+	static UE_API void CategoryCreated(uint8 Category, const ANSICHAR* Name);
+	static UE_API void RequestStarted(uint64 Dispatcher, uint64 Request, const ANSICHAR* Url, int32 Priority, uint8 Category);
+	static UE_API void ChunkRangeAdded(uint64 Request, const class FIoChunkId& ChunkId, uint32 Start, uint32 End);
+	static UE_API void RequestCompleted(uint64 Request, const ANSICHAR* Host, uint32 StatusCode, uint32 ContentLength);
+};
+
+////////////////////////////////////////////////////////////////////////////////////////////////////
+#define TRACE_HTTP_DISPATCHER_CREATED(Dispatcher, Name) \
+	FHttpTrace::DispatcherCreated(Dispatcher, Name);
+
+#define TRACE_HTTP_CATEGORY_CREATED(Category, Name) \
+	FHttpTrace::CategoryCreated(Category, Name);
+
+#define TRACE_HTTP_REQUEST_STARTED(Dispatcher, Request, Url, Priority, Category) \
+	FHttpTrace::RequestStarted(Dispatcher, Request, Url, Priority, Category);
+
+#define TRACE_HTTP_CHUNK_RANGE_ADDED(Request, ChunkId, Start, End) \
+	FHttpTrace::ChunkRangeAdded(Request, ChunkId, Start, End);
+
+#define TRACE_HTTP_REQUEST_COMPLETED(Request, Host, StatusCode, ContentLength) \
+	FHttpTrace::RequestCompleted(Request, Host, StatusCode, ContentLength);
+
+////////////////////////////////////////////////////////////////////////////////////////////////////
+#else
+
+#define TRACE_HTTP_DISPATCHER_CREATED(...)
+#define TRACE_HTTP_CATEGORY_CREATED(Category, Name)(...)
+#define TRACE_HTTP_REQUEST_STARTED(...)
+#define TRACE_HTTP_CHUNK_RANGE_ADDED(Request, ChunkId, Start, End)(...)
+#define TRACE_HTTP_REQUEST_COMPLETED(...)
+
+#endif //UE_TRACE_HTTP_ENABLED
+
+#undef UE_API
