Added support for filtering HTTP request(s) based on the selected time range in Insights timing view.

#rnx

#Synced-CL 50547167

diff --git a/Engine/Plugins/Experimental/HttpInsights/Source/Private/View.cpp b/Engine/Plugins/Experimental/HttpInsights/Source/Private/View.cpp
--- a/Engine/Plugins/Experimental/HttpInsights/Source/Private/View.cpp
+++ b/Engine/Plugins/Experimental/HttpInsights/Source/Private/View.cpp
@@ -9,6 +9,8 @@
 #include "Insights/ITimingViewExtender.h"
 #include "Insights/ITimingViewSession.h"
 #include "Math/Color.h"
+#include "Math/Interval.h"
+#include "Math/UnrealMathUtility.h"
 #include "Misc/EnumClassFlags.h"
 #include "Widgets/DeclarativeSyntaxSupport.h"
 #include "Widgets/Docking/SDockTab.h"
@@ -459,6 +461,7 @@ public:
 								SLATE_BEGIN_ARGS(SHttpFilterBar) { }
 								SLATE_END_ARGS()
 	void						Construct(const FArguments& Args);
+	void						SetSelectedTimeRange(FDoubleInterval TimeRange);
 	bool						ShouldFilter(const FHttpRequestViewModel& Vm) const;
 	FFilterChangedDelegate&		OnFilterChanged() { return FilterChanged; }
 
@@ -471,6 +474,7 @@ private:
 	TSharedPtr<SEditableTextBox>	SearchTextBox;
 	FFilterChangedDelegate			FilterChanged;
 	EHttpCategoryFilter				CurrentCategoryFilter = EHttpCategoryFilter::All;
+	FDoubleInterval					CurrentTimeRange;
 };
 
 ////////////////////////////////////////////////////////////////////////////////////////////////////
@@ -514,6 +518,12 @@ void SHttpFilterBar::Construct(const FArguments& Args)
 	];
 }
 
+void SHttpFilterBar::SetSelectedTimeRange(FDoubleInterval TimeRange)
+{
+	CurrentTimeRange = TimeRange;
+	FilterChanged.Broadcast();
+}
+
 bool SHttpFilterBar::ShouldFilter(const FHttpRequestViewModel& Vm) const
 {
 	//TODO: Category filtering is currently hardcoded see OnDemandHttpIoDispatcher.h
@@ -527,6 +537,14 @@ bool SHttpFilterBar::ShouldFilter(const FHttpRequestViewModel& Vm) const
 		return true;
 	}
 
+	if (CurrentTimeRange.IsValid())
+	{
+		if (CurrentTimeRange.Contains(Vm->Request.CompletionTime) == false)
+		{
+			return true;
+		}
+	}
+
 	const FText SearchText = SearchTextBox->GetText();
 	if (SearchText.IsEmpty() == false)
 	{
@@ -632,6 +650,9 @@ public:
 	virtual void			Tick(const FGeometry& AllottedGeometry, const double InCurrentTime, const float InDeltaTime) override;
 
 private:
+	void					OnTimeMarkerChanged(ETimeChangedFlags Flags, double TimeMarker);
+	void					OnTimeSelectionChanged(ETimeChangedFlags Flags, double StartTime, double EndTime);
+
 	FHttpTimingViewExtender*			TimingViewExtender = nullptr;
 	FHttpTimingViewSession*				CurrentSession = nullptr;
 	TArray<FHttpRequestViewModel>		HttpRequests;
@@ -646,6 +667,14 @@ const FName SHttpInsightsTab::TabId = FName("HttpInsightsTab");
 
 SHttpInsightsTab::~SHttpInsightsTab()
 {
+	if (CurrentSession != nullptr)
+	{
+		CurrentSession->GetTimingViewSession()
+			.OnTimeMarkerChanged().RemoveAll(this);
+		CurrentSession->GetTimingViewSession()
+			.OnSelectionChanged().RemoveAll(this);
+	}
+
 	if (TimingViewExtender != nullptr)
 	{
 		TimingViewExtender->RemoveView(this);
@@ -683,12 +712,20 @@ void SHttpInsightsTab::OnBeginSession(FHttpTimingViewSession& Session)
 {
 	check(CurrentSession == nullptr);
 	CurrentSession = &Session;
+	CurrentSession->GetTimingViewSession()
+		.OnTimeMarkerChanged().AddRaw(this, &SHttpInsightsTab::OnTimeMarkerChanged);
+	CurrentSession->GetTimingViewSession()
+		.OnSelectionChanged().AddRaw(this, &SHttpInsightsTab::OnTimeSelectionChanged);
 }
 
 void SHttpInsightsTab::OnEndSession(FHttpTimingViewSession& Session)
 {
 	if (CurrentSession == &Session)
 	{
+		CurrentSession->GetTimingViewSession()
+			.OnTimeMarkerChanged().RemoveAll(this);
+		CurrentSession->GetTimingViewSession()
+			.OnSelectionChanged().RemoveAll(this);
 		CurrentSession = nullptr;
 		HttpRequests.Empty();
 		AllHttpRequests.Empty();
@@ -743,6 +780,20 @@ void SHttpInsightsTab::Tick(const FGeometry& AllottedGeometry, const double InCu
 	}
 }
 
+void SHttpInsightsTab::OnTimeMarkerChanged(ETimeChangedFlags Flags, double TimeMarker)
+{
+	//TODO: Find closest entry to time marker
+}
+
+void SHttpInsightsTab::OnTimeSelectionChanged(ETimeChangedFlags Flags, double StartTime, double EndTime)
+{
+	if (Flags == UE::Insights::Timing::ETimeChangedFlags::None)
+	{
+		const bool bValidRange = (EndTime > StartTime) && !FMath::IsNearlyEqual(StartTime, EndTime);
+		FilterBar->SetSelectedTimeRange(bValidRange ? FDoubleInterval(StartTime, EndTime) : FDoubleInterval());
+	}
+}
+
 ////////////////////////////////////////////////////////////////////////////////////////////////////
 void RegisterTimingProfilerExtensions(
 	FInsightsMajorTabExtender& InOutExtender,
