Remove mismatched loc text define

#rnx

#Synced-CL 50723010

diff --git a/Engine/Plugins/Experimental/HttpInsights/Source/Private/Model.cpp b/Engine/Plugins/Experimental/HttpInsights/Source/Private/Model.cpp
--- a/Engine/Plugins/Experimental/HttpInsights/Source/Private/Model.cpp
+++ b/Engine/Plugins/Experimental/HttpInsights/Source/Private/Model.cpp
@@ -247,5 +247,3 @@ TSharedPtr<IHttpLogModel> MakeHttpLogModel(IAnalysisSession& Session)
 }
 
 } // namespace UE:HttpInsights
-
-#undef LOCTEXT_NAMESPACE
